Copied cl 9491177

AxFImporter: Fixed import of materials with '.' in their filenames(sanitized material name); Better message for unsupported anisotropy, fixed double negative when unexpected texture processed.

diff --git a/Engine/Plugins/Enterprise/AxFImporter/Source/AxFImporter/Private/AxFImporter.cpp b/Engine/Plugins/Enterprise/AxFImporter/Source/AxFImporter/Private/AxFImporter.cpp
--- a/Engine/Plugins/Enterprise/AxFImporter/Source/AxFImporter/Private/AxFImporter.cpp
+++ b/Engine/Plugins/Enterprise/AxFImporter/Source/AxFImporter/Private/AxFImporter.cpp
@@ -6,6 +6,7 @@
 
 
 #include "PackageTools.h"
+#include "ObjectTools.h"
 #include "AssetRegistryModule.h"
 #include "EditorFramework/AssetImportData.h"
 #include "Engine/Texture2D.h"
@@ -1472,7 +1473,8 @@ private:
 		{
 			ICreatedMaterial& Material = *MaterialPtr.Get();
 
-			FString MaterialName = Material.Name;
+			FString MaterialName =  ObjectTools::SanitizeObjectName(Material.Name);
+
 			Log.Info(TEXT("Importing material: ") + MaterialName);
 
 			FString  MaterialPackageName = UPackageTools::SanitizePackageName(*(ParentPackage->GetName() / MaterialName));
@@ -1926,9 +1928,13 @@ private:
 // 					{
 // 						Material.SetTextureCarpaintClearcoatNormal(ProcessedTextureSource);
 // 					}
+					else if (TextureName == AXF_SVBRDF_TEXTURE_NAME_ANISO_ROTATION)
+					{
+						Log.Warn(TEXT("Anisotropic AxF materials are supported yet - will use isotropic approximation."));
+					}
 					else
 					{
-						Log.Warn(TEXT("Texture \"") + TextureName + TEXT("\" wasn't not handled"));
+						Log.Warn(TEXT("Texture \"") + TextureName + TEXT("\" wasn't handled"));
 					}
 				}
 
@@ -2259,13 +2265,14 @@ private:
 		for (auto& CreatedMaterialPtr : CreatedMaterials)
 		{
 			ICreatedMaterial& CreatedMaterial = *CreatedMaterialPtr.Get();
-			if (CreatedMaterial.Name == OutMaterial->GetName())
+			FString MaterialName = ObjectTools::SanitizeObjectName(CreatedMaterial.Name);
+			if (MaterialName == OutMaterial->GetName())
 			{
 				FString RootPackageName = *FPaths::GetPath(OutMaterial->GetOuter()->GetName());
 				UPackage* RootPackage = CreatePackage(nullptr, *RootPackageName);
 				
 				ImportMaterial(CreatedMaterial, Material, RootPackage);
-				Log.Info(TEXT("Done re-importing material: ") + CreatedMaterial.Name);
+				Log.Info(TEXT("Done re-importing material: ") + MaterialName);
 				return true;
 			}
 		}
