Unshelved from pending changelist '9317935'

"AxF - Improve CarPaint Roughness/F0 fitting" -  UEENT-2942

diff --git a/Engine/Plugins/Enterprise/AxFImporter/Source/AxFImporter/Private/AxFImporter.cpp b/Engine/Plugins/Enterprise/AxFImporter/Source/AxFImporter/Private/AxFImporter.cpp
--- a/Engine/Plugins/Enterprise/AxFImporter/Source/AxFImporter/Private/AxFImporter.cpp
+++ b/Engine/Plugins/Enterprise/AxFImporter/Source/AxFImporter/Private/AxFImporter.cpp
@@ -57,6 +57,7 @@
 #include "AxFImporterOptions.h"
 
 #include "ImageCore.h"
+#include "Async/ParallelFor.h"
 
 class FAxFLog {
 public:
@@ -665,8 +666,8 @@ public:
 			return Result;
 		}
 
-		template<class Input>
-		void Connect(Input& Input, FMaterialExpressionWrapper V)
+		template<class InputType>
+		void Connect(InputType& Input, FMaterialExpressionWrapper V)
 		{
 			check(V.Expression);
 			Input.Connect(V.OutputIndex, V.Expression);
@@ -686,6 +687,108 @@ public:
 		bool bNoRefraction = true;
 	};
 
+	class FUeGgxBRDF
+	{
+	public:
+
+		float Saturate(float V)
+		{
+			return FMath::Max(0.0f, FMath::Min(V, 1.0f));
+		}
+
+		// from BRDF.ush
+		// GGX / Trowbridge-Reitz
+		// [Walter et al. 2007, "Microfacet models for refraction through rough surfaces"]
+		float D_GGX(float a2, float NoH)
+		{
+			float d = (NoH * a2 - NoH) * NoH + 1;	// 2 mad
+			return a2 / (PI * d * d);					// 4 mul, 1 rcp
+		}
+
+		// Appoximation of joint Smith term for GGX
+		// [Heitz 2014, "Understanding the Masking-Shadowing Function in Microfacet-Based BRDFs"]
+		float Vis_SmithJointApprox(float a2, float NoV, float NoL)
+		{
+			float a = sqrt(a2);
+			float Vis_SmithV = NoL * (NoV * (1 - a) + a);
+			float Vis_SmithL = NoV * (NoL * (1 - a) + a);
+			return 0.5f * 1.0f /(Vis_SmithV + Vis_SmithL);
+		}
+
+		// [Schlick 1994, "An Inexpensive BRDF Model for Physically-Based Rendering"]
+		FVector F_Schlick(FVector SpecularColor, float VoH)
+		{
+			
+			float Fc = FMath::Pow(1 - VoH, 5);					// 1 sub, 3 mul
+			//return Fc + (1 - Fc) * SpecularColor;		// 1 add, 3 mad
+
+			// Anything less than 2% is physically impossible and is instead considered to be shadowing
+			return FVector(Saturate(50.0 * SpecularColor.Y) * Fc) + (1 - Fc) * SpecularColor;
+
+		}
+
+		// from ShadingModels.ush
+		FVector SpecularGGX(float Roughness, FVector SpecularColor, float NoV, float NoL, float NoH, float VoH)
+		{
+			float a2 = FMath::Pow(Roughness, 4);
+			//float Energy = EnergyNormalization(a2, Context.VoH, AreaLight);
+			float Energy = 1.0f;
+
+			// Generalized microfacet specular
+			float D = D_GGX(a2, NoH) * Energy;
+			float Vis = Vis_SmithJointApprox(a2, NoV, NoL);
+			FVector F = F_Schlick(SpecularColor, VoH);
+
+			return (D * Vis) * F;
+		}
+
+	};
+
+	class FCarPaint2BRDF
+	{
+	public:
+		struct FParams
+		{
+			float Diffuse;
+
+			float Spreads[3];
+			float Coeffs[3];
+			float F0s[3];
+		};
+
+		float CT_F(float H_V, float F0)
+		{
+			return F0 + (1.0 - F0) * FMath::Pow(1.0 - H_V, 5);
+		}
+
+		// https://en.wikipedia.org/wiki/Specular_highlight#Beckmann_distribution
+		float CT_D(float N_H, float m)
+		{
+			float CosSqr = FMath::Square(N_H);
+			return FMath::Exp((CosSqr - 1.0) / (CosSqr * FMath::Square(m))) / (FMath::Square(m) * FMath::Square(CosSqr));
+		}
+
+		float CT_G(float NoV, float NoL, float NoH, float HoV)
+		{
+			return FMath::Min(1.0f, FMath::Min((2.0f * NoH * NoV / HoV), (2.0f * NoH * NoL / HoV)));
+		}
+
+		float CT(const FParams& params, int LobeCount, float NoV, float NoL, float NoH, float HoV)
+		{
+			float Result = params.Diffuse / PI;
+
+			// specular term
+			float Scale = CT_G(NoH, NoV, NoL, HoV) / (PI * NoV * NoL);
+			for (int32 LobeIndex = 0; LobeIndex < LobeCount; ++LobeIndex)
+			{
+				float Spread = params.Spreads[LobeIndex];
+				Result += params.Coeffs[LobeIndex] * CT_D(NoH, Spread) * CT_F(HoV, params.F0s[LobeIndex]) * Scale;
+			}
+			return Result;
+		}
+
+	};
+
 	class FCreatedCarpaint2Material: public ICreatedMaterial
 	{
 	public:
@@ -711,8 +814,10 @@ public:
 
 		bool bEnableClearcoat = true;
 
-		void Build()
+		void Build() override
 		{
+			FitBRDF();
+
 			auto BaseColor = Box(nullptr);
 			auto Specular = Box(nullptr);
 
@@ -753,22 +858,8 @@ public:
 				}
 			}
 
-			float RoughnessSum = 0;
-			float F0Sum = 0;
-			float CT_CoeffsSum = 0;
-			for (int i = 0; i < 3; ++i)
-			{
-				RoughnessSum += CT_Spreads[i] * CT_Coeffs[i];
-				F0Sum += CT_F0s[i] * CT_Coeffs[i];
-				CT_CoeffsSum += CT_Coeffs[i];
-			}
-
-			float F0 = F0Sum / CT_CoeffsSum;
-			float Roughness = RoughnessSum / CT_CoeffsSum;
-
-			Connect(Material->Metallic, Constant(F0));
-
-			Connect(Material->Roughness, Constant(Roughness));
+			Connect(Material->Metallic, Constant(1.0f));
+			Connect(Material->Roughness, Constant(RoughnessFitted));
 
 			if (bHasBRDFColorsTexture || bHasBTFFlakes)
 			{
@@ -815,10 +906,12 @@ public:
 
 				if (bHasBRDFColorsTexture)
 				{
-					// ??? why normalize  by pi/2 when acos is 0..pi?
 					auto UV = Append(Mul(ThetaF, 1.f / HALF_PI), Mul(ThetaI, 1.f / HALF_PI));
-					// Add F0 to compensate to reflection power loss
-					BaseColor = Mul(SampleTexture(BRDFColorsTexture, UV), CT_Diffuse + F0Sum);
+					BaseColor = Mul(SampleTexture(BRDFColorsTexture, UV), SpecularityFitted);
+				}
+				else
+				{
+					BaseColor = Constant(SpecularityFitted);
 				}
 
 				if (bHasBTFFlakes)
@@ -889,7 +982,7 @@ public:
 				Connect(Material->Specular, Specular);
 			}
 
-			bool bHasClearcoat = ExpressionClearcoatIOR || ExpressionClearcoatNormal || ExpressionClearcoatColor;
+			bool bHasClearcoat = ExpressionClearcoatIOR || ExpressionClearcoatNormal || ExpressionClearcoatColor || (CC_IOR != 1.0f);
 			// seems like AxF decoding sdk doesn't provide any other clues when material has ClearCoat layer
 			if (bHasClearcoat)
 			{
@@ -909,6 +1002,8 @@ public:
 					UMaterialExpressionClearCoatNormalCustomOutput* BottomNormal = CreateMaterialExpression< UMaterialExpressionClearCoatNormalCustomOutput>();
 					Connect(BottomNormal->Input, ExpressionNormal);
 				}
+
+				Connect(Material->ClearCoatRoughness, Constant(0.0f));
 			}
 			else
 			{
@@ -919,6 +1014,145 @@ public:
 			}
 		}
 
+		float RoughnessFitted;
+		float SpecularityFitted;
+
+		// Finds UE GGX shading model parameters(roughness and Specularity) which fit best
+		// source AxF material's BRDF(X-Rite CPA2)
+		void FitBRDF()
+		{
+			// BRDFs are parametrized by angles between light/normal(ThetaL) and view/normal(ThetaV) over hemisphere around surface normal.
+			const int32 LVGridSize = 100; 
+			const int32 RougnessSpecularitySearchAreaSize = 30; // Search begins with whole range of roughness/specularity using this resolution
+			const float IterationRangeReduction = 0.2; // After finding best fit search restarts around found value in a smaller area
+			const int32 IterationCount = 5;
+			const float RayAngleExtent = HALF_PI * 0.8; // don't take into account angles too close to grazing - spikes are too big there 
+			const bool bUseL2Distance = true;
+
+			static float VisibleDirections[LVGridSize];
+
+			// use light/view directions around normal from -pi/2 to pi/2
+			for (int I = 0; I < LVGridSize; ++I)
+			{
+				float Alpha = float(I) / (LVGridSize - 1);
+				VisibleDirections[I] = -RayAngleExtent * (1.0f - Alpha) + RayAngleExtent * Alpha;
+			}
+
+			static float ThetaL[LVGridSize][LVGridSize];
+			static float ThetaV[LVGridSize][LVGridSize];
+			static float ThetaH[LVGridSize][LVGridSize];
+
+			static float N_V[LVGridSize][LVGridSize];
+			static float N_H[LVGridSize][LVGridSize];
+			static float N_L[LVGridSize][LVGridSize];
+			static float H_V[LVGridSize][LVGridSize];
+
+			for (int U = 0; U < LVGridSize; ++U)
+			{
+				for (int V = 0; V < LVGridSize; ++V)
+				{
+					ThetaL[U][V] = VisibleDirections[U];
+					ThetaV[U][V] = VisibleDirections[V];
+					ThetaH[U][V] = (ThetaL[U][V] + ThetaV[U][V]) * 0.5f; // half-way vector(in the same plane) 
+
+					N_L[U][V] = FMath::Cos(ThetaL[U][V]);
+					N_V[U][V] = FMath::Cos(ThetaV[U][V]);
+					N_H[U][V] = FMath::Cos(ThetaH[U][V]);
+					H_V[U][V] = FMath::Cos(ThetaH[U][V] - ThetaV[U][V]);
+				}
+			}
+
+			FCarPaint2BRDF CarPaint2BRDF;
+			FUeGgxBRDF UeGgxBRDF;
+
+			FCarPaint2BRDF::FParams Params;
+			Params.Diffuse = CT_Diffuse;
+			for (int32 Lobe = 0; Lobe < 3; ++Lobe)
+			{
+				Params.Coeffs[Lobe] = Lobe < CT_Coeffs.Num() ? CT_Coeffs[Lobe] : 0.0f;
+				Params.Spreads[Lobe] = Lobe < CT_Spreads.Num() ? CT_Spreads[Lobe] : 1.0f;
+				Params.F0s[Lobe] = Lobe < CT_F0s.Num() ? CT_F0s[Lobe] : 1.0f;
+			}
+
+			float RoughnessRange[2] = { 0.01f, 1.0f };
+			float SpecularityRange[2] = { 0.0f, 1.0f };
+
+			float ClosestDistance = BIG_NUMBER;
+			float ClosestRoughness = 1.0f;
+			float ClosestSpecularity = 1.0f;
+			for (int Iteration = 0; Iteration < IterationCount; ++Iteration)
+			{
+
+				for (int32 RoughnessIndex = 0; RoughnessIndex <= RougnessSpecularitySearchAreaSize; ++RoughnessIndex)
+				{
+					float Specularities[RougnessSpecularitySearchAreaSize];
+					float DistanceForSpecularity[RougnessSpecularitySearchAreaSize];
+
+					float Roughness = FMath::Lerp(RoughnessRange[0], RoughnessRange[1], float(RoughnessIndex) / RougnessSpecularitySearchAreaSize);
+
+					ParallelFor(RougnessSpecularitySearchAreaSize + 1,
+						[&](int32 SpecularityIndex)
+						{
+							float Specularity = FMath::Lerp(SpecularityRange[0], SpecularityRange[1], float(SpecularityIndex) / RougnessSpecularitySearchAreaSize);
+
+							float SumOfSquaredDiff = 0.0f;
+							float SumOfAbsDiff = 0.0f;
+							float SumOfSqrtDiff = 0.0f;
+
+							for (int U = 0; U < LVGridSize; ++U)
+							{
+								for (int V = 0; V < LVGridSize; ++V)
+								{
+									float NoL = N_L[U][V];
+									float NoV = N_V[U][V];
+									float NoH = N_H[U][V];
+									float HoV = H_V[U][V];
+
+									float xrite = NoL*CarPaint2BRDF.CT(Params, 3, NoV, NoL, NoH, HoV);
+									FVector ue = NoL*UeGgxBRDF.SpecularGGX(Roughness, FVector(Specularity), NoV, NoL, NoH, HoV);
+
+									SumOfSquaredDiff += FMath::Square(FMath::Max(xrite, 0.0f) - FMath::Max(ue.X, 0.0f));
+									SumOfAbsDiff += FMath::Abs(FMath::Max(xrite, 0.0f) - FMath::Max(ue.X, 0.0f));
+									SumOfSqrtDiff += FMath::Sqrt(FMath::Abs(FMath::Max(xrite, 0.0f) - FMath::Max(ue.X, 0.0f)));
+								}
+							}
+
+							//float Distance = SumOfSquaredDiff;
+							//float Distance = SumOfAbsDiff;
+							float Distance = SumOfAbsDiff;
+
+							Specularities[SpecularityIndex] = Specularity;
+							DistanceForSpecularity[SpecularityIndex] = Distance;
+						}
+					);
+
+					for (int32 SpecularityIndex = 0; SpecularityIndex <= RougnessSpecularitySearchAreaSize; ++SpecularityIndex)
+					{
+						float Specularity = Specularities[SpecularityIndex];
+						float Distance = DistanceForSpecularity[SpecularityIndex];
+						if (Distance < ClosestDistance)
+						{
+						 	ClosestDistance = Distance;
+						 	ClosestRoughness = Roughness;
+						 	ClosestSpecularity = Specularity;
+						}
+					}
+				}
+
+				// reduce search ranges
+				float RoughnessSpread = (RoughnessRange[1] - RoughnessRange[0]) * IterationRangeReduction;
+				RoughnessRange[0] = FMath::Max(0.0f, ClosestRoughness - RoughnessSpread);
+				RoughnessRange[1] = FMath::Min(1.0f, ClosestRoughness + RoughnessSpread);
+
+				float SpecularitySpread = (SpecularityRange[1] - SpecularityRange[0]) * IterationRangeReduction;
+				SpecularityRange[0] = FMath::Max(0.0f, ClosestSpecularity - SpecularitySpread);
+				SpecularityRange[1] = FMath::Min(1.0f, ClosestSpecularity + SpecularitySpread);
+			}
+
+			RoughnessFitted = ClosestRoughness;
+			SpecularityFitted = ClosestSpecularity;
+		}
+
 		FMaterialExpressionWrapper SampleFlake(FMaterialExpressionWrapper ll_uv, FMaterialExpressionWrapper ll_slice)
 		{
 // 			CreateFunction(TEXT("SampleFlakeTexture"), [this](UMaterialFunction* Function) 
@@ -2383,3 +2617,4 @@ bool FAxFImporter::IsLoaded()
 
 #endif
 
+
