Fix out of bound error.

diff --git a/Engine/Plugins/Enterprise/AxFImporter/Source/AxFImporter/Private/AxFImporter.cpp b/Engine/Plugins/Enterprise/AxFImporter/Source/AxFImporter/Private/AxFImporter.cpp
--- a/Engine/Plugins/Enterprise/AxFImporter/Source/AxFImporter/Private/AxFImporter.cpp
+++ b/Engine/Plugins/Enterprise/AxFImporter/Source/AxFImporter/Private/AxFImporter.cpp
@@ -1083,17 +1083,17 @@ public:
 			for (int Iteration = 0; Iteration < IterationCount; ++Iteration)
 			{
 
-				for (int32 RoughnessIndex = 0; RoughnessIndex <= RougnessSpecularitySearchAreaSize; ++RoughnessIndex)
+				for (int32 RoughnessIndex = 0; RoughnessIndex < RougnessSpecularitySearchAreaSize; ++RoughnessIndex)
 				{
 					float Specularities[RougnessSpecularitySearchAreaSize];
 					float DistanceForSpecularity[RougnessSpecularitySearchAreaSize];
 
-					float Roughness = FMath::Lerp(RoughnessRange[0], RoughnessRange[1], float(RoughnessIndex) / RougnessSpecularitySearchAreaSize);
+					float Roughness = FMath::Lerp(RoughnessRange[0], RoughnessRange[1], float(RoughnessIndex) / (RougnessSpecularitySearchAreaSize-1));
 
-					ParallelFor(RougnessSpecularitySearchAreaSize + 1,
+					ParallelFor(RougnessSpecularitySearchAreaSize,
 						[&](int32 SpecularityIndex)
 						{
-							float Specularity = FMath::Lerp(SpecularityRange[0], SpecularityRange[1], float(SpecularityIndex) / RougnessSpecularitySearchAreaSize);
+							float Specularity = FMath::Lerp(SpecularityRange[0], SpecularityRange[1], float(SpecularityIndex) / (RougnessSpecularitySearchAreaSize-1));
 
 							float SumOfSquaredDiff = 0.0f;
 							float SumOfAbsDiff = 0.0f;
@@ -1126,7 +1126,7 @@ public:
 						}
 					);
 
-					for (int32 SpecularityIndex = 0; SpecularityIndex <= RougnessSpecularitySearchAreaSize; ++SpecularityIndex)
+					for (int32 SpecularityIndex = 0; SpecularityIndex < RougnessSpecularitySearchAreaSize; ++SpecularityIndex)
 					{
 						float Specularity = Specularities[SpecularityIndex];
 						float Distance = DistanceForSpecularity[SpecularityIndex];
