Update path to AxF libs

diff --git a/Engine/Plugins/Enterprise/AxFImporter/Source/AxFImporter/AxFImporter.Build.cs b/Engine/Plugins/Enterprise/AxFImporter/Source/AxFImporter/AxFImporter.Build.cs
--- a/Engine/Plugins/Enterprise/AxFImporter/Source/AxFImporter/AxFImporter.Build.cs
+++ b/Engine/Plugins/Enterprise/AxFImporter/Source/AxFImporter/AxFImporter.Build.cs
@@ -10,7 +10,7 @@ namespace UnrealBuildTool.Rules
     {
         private string ThirdPartyPath
         {
-            get { return Path.GetFullPath(Path.Combine(EngineDirectory, "Restricted/NotForLicensees/Plugins/Enterprise/AxFImporter/Source/ThirdParty")); }
+            get { return Path.GetFullPath(Path.Combine(EngineDirectory, "Restricted/NotForLicensees/Source/ThirdParty/Enterprise")); }
         }
 
         public AxFImporter(ReadOnlyTargetRules Target) : base(Target)
