Updated AxF decoding SDK to 1.8.1

diff --git a/Engine/Plugins/Enterprise/AxFImporter/Source/AxFImporter/AxFImporter.Build.cs b/Engine/Plugins/Enterprise/AxFImporter/Source/AxFImporter/AxFImporter.Build.cs
--- a/Engine/Plugins/Enterprise/AxFImporter/Source/AxFImporter/AxFImporter.Build.cs
+++ b/Engine/Plugins/Enterprise/AxFImporter/Source/AxFImporter/AxFImporter.Build.cs
@@ -45,10 +45,14 @@ namespace UnrealBuildTool.Rules
                 }
             );
 
+			string AfxSdkVersion = "1.8.1";
+
             if (Target.Platform == UnrealTargetPlatform.Win64)
             {
-                PublicDelayLoadDLLs.Add("AxFDecoding.dll");
-				string ModulePath = Path.Combine(EngineDirectory, "Plugins/Enterprise/AxFImporter/Binaries/ThirdParty/AxF", Target.Platform.ToString(), "AxFDecoding.dll");
+				string AxFDllName = $"AxFDecoding.{AfxSdkVersion}.dll";
+
+				PublicDelayLoadDLLs.Add(AxFDllName);
+				string ModulePath = Path.Combine(EngineDirectory, "Plugins/Enterprise/AxFImporter/Binaries/ThirdParty/AxF", Target.Platform.ToString(), AxFDllName);
 				if (!File.Exists(ModulePath))
                 {
 					string Err = string.Format("AxF Decoding dll '{0}' not found.", ModulePath);
@@ -62,7 +66,7 @@ namespace UnrealBuildTool.Rules
             {
                 //third party libraries
 
-                string[] Libs = { "AxF-Decoding-SDK-1.5.1" };
+                string[] Libs = { $"AxF-Decoding-SDK-{AfxSdkVersion}" };
                 string[] StaticLibNames = { "AxFDecoding" };
 
                 foreach (string Lib in Libs)
@@ -78,9 +82,10 @@ namespace UnrealBuildTool.Rules
                     }
                 }
 
-                PrivateDefinitions.Add("USE_AXFSDK");
+				PrivateDefinitions.Add("USE_AXFSDK");
+				PrivateDefinitions.Add($"AFX_SDK_VERSION=\"{AfxSdkVersion}\"");
 
-                string TargetPlatform = "Windows.x64";
+				string TargetPlatform = "Windows.x64";
                 string TargetExtension = "lib";
                 if (Target.Platform == UnrealTargetPlatform.Win64)
                 {
diff --git a/Engine/Plugins/Enterprise/AxFImporter/Source/AxFImporter/Private/AxFImporter.cpp b/Engine/Plugins/Enterprise/AxFImporter/Source/AxFImporter/Private/AxFImporter.cpp
--- a/Engine/Plugins/Enterprise/AxFImporter/Source/AxFImporter/Private/AxFImporter.cpp
+++ b/Engine/Plugins/Enterprise/AxFImporter/Source/AxFImporter/Private/AxFImporter.cpp
@@ -4,6 +4,10 @@
 
 #ifdef USE_AXFSDK
 
+// To avoid annoying VS' visual cue that AFX_SDK_VERSION is not defined
+#ifndef AFX_SDK_VERSION
+#define AFX_SDK_VERSION
+#endif
 
 #include "PackageTools.h"
 #include "ObjectTools.h"
@@ -1618,7 +1622,8 @@ private:
 
 			Log.Info(TEXT("Material name: ") + MaterialName);
 
-			axf::decoding::AXF_REPRESENTATION_HANDLE RepresentationHandle = axf::decoding::axfGetPreferredRepresentation(MaterialHandle);
+			// Look for using axfGetBestCompatibleRepresentation
+			axf::decoding::AXF_REPRESENTATION_HANDLE RepresentationHandle = axf::decoding::axfGetBestSDKSupportedRepresentation(MaterialHandle);
 			if (!RepresentationHandle)
 			{
 				Log.Error(FString::Printf(TEXT("Can't retrieve preferred representation for material %d('%s')"), MaterialIndex, *MaterialName));
@@ -2003,11 +2008,6 @@ private:
 						// 1x1 or constant on samples
 					}
 
-					if (TextureName == AXF_SVBRDF_TEXTURE_NAME_CLEARCOAT_COLOR)
-					{
-						// not present on samples
-					}
-
 					bool bReplaceTextureByConstant = (Width == 1) && (Height == 1);
 
 
@@ -2152,10 +2152,6 @@ private:
 					{
 						Material.SetTextureClearcoatIOR(ProcessedTextureSource);
 					}
-					else if (TextureName == AXF_SVBRDF_TEXTURE_NAME_CLEARCOAT_COLOR)
-					{
-						Material.SetTextureClearcoatColor(ProcessedTextureSource);
-					}
 					else if (TextureName == AXF_SVBRDF_TEXTURE_NAME_HEIGHT)
 					{
 						Material.SetTextureHeight(ProcessedTextureSource);
@@ -2597,7 +2593,8 @@ FAxFImporter::FAxFImporter(const FString& PluginPath)
 	FString LibraryPath = FPaths::Combine(ThirdPartyPath, TEXT("AxF"), Platform);
 
 	FPlatformProcess::PushDllDirectory(*LibraryPath);
-	AxFDecodingHandle = FPlatformProcess::GetDllHandle(TEXT("AxFDecoding.dll"));
+	FString DllName = TEXT("AxFDecoding.") + FString(AFX_SDK_VERSION) + TEXT(".dll");
+	AxFDecodingHandle = FPlatformProcess::GetDllHandle(*DllName);
 	FPlatformProcess::PopDllDirectory(*LibraryPath);
 
 	axf::decoding::axfEnableLogging(FAxFFileImporter::LoggingCallback, axf::decoding::LOGLEVEL_INFO);
@@ -2618,7 +2615,7 @@ bool FAxFImporter::IsLoaded()
 	return AxFDecodingHandle != nullptr;
 }
 
-#else
+#else // USE_AXFSDK
 
 FAxFImporter::FAxFImporter(const FString& PluginPath)
 {
@@ -2638,6 +2635,6 @@ bool FAxFImporter::IsLoaded()
 	return false;
 }
 
-#endif
+#endif // USE_AXFSDK
 
 
