nop just comments for FImage todos

diff --git a/Engine/Plugins/Enterprise/AxFImporter/Source/AxFImporter/Private/AxFImporter.cpp b/Engine/Plugins/Enterprise/AxFImporter/Source/AxFImporter/Private/AxFImporter.cpp
--- a/Engine/Plugins/Enterprise/AxFImporter/Source/AxFImporter/Private/AxFImporter.cpp
+++ b/Engine/Plugins/Enterprise/AxFImporter/Source/AxFImporter/Private/AxFImporter.cpp
@@ -99,6 +99,8 @@ public:
 
 	struct FAxFTextureSource
 	{
+		// @@!! just use float here, convert later
+		//  ideally replace this whole class with FImage
 		TArray<FFloat16> Pixels;
 
 		int32 Width;
@@ -2248,6 +2250,9 @@ private:
 			PixelValueMin = PixelValueMin.ComponentMin(PixelValue);
 		}
 
+		// @@!!
+		// float image to RGBA16
+		// todo: use FImage / CopyImage
 		TArray<uint16> PixelsCompressed;
 		if (!PixelValueMax.IsNearlyZero())
 		{
@@ -2266,8 +2271,10 @@ private:
 
 				uint16 One = 65535;
 
+				// @@!! precompute 1/Scale then clamp in [0,1]
 				FVector C = (ClampVector(Color, FVector::ZeroVector, Scale) / Scale) * One; // clamp, to avoid integer overflow(which results in nasty pixels)
 
+				// @@!! incorrect conversion; use QuantizeUNormFloatTo16
 				PixelsCompressed[PixelIndex * 4 + 0] = FMath::FloorToInt(C.X);
 				PixelsCompressed[PixelIndex * 4 + 1] = FMath::FloorToInt(C.Y);
 				PixelsCompressed[PixelIndex * 4 + 2] = FMath::FloorToInt(C.Z);
@@ -2365,6 +2372,7 @@ private:
 
 	void SetTextureSource(UTexture2D* Texture, FAxFTextureSource TextureSource)
 	{
+		//@@!! can I just make 32F ? no conversion?
 		Texture->Source.Init(
 			TextureSource.Width,
 			TextureSource.Height,
@@ -2379,6 +2387,7 @@ private:
 	{
 		FImage Image;
 		Image.Init(TextureSource.Width, TextureSource.Height, ERawImageFormat::RGBA32F);
+		// @@!! if TextureSource is FImage this goes away
 		for (int PixelIndex = 0; PixelIndex < TextureSource.Width * TextureSource.Height; PixelIndex++)
 		{
 			Image.AsRGBA32F()[PixelIndex] = TextureSource.GetPixel(PixelIndex);
@@ -2403,6 +2412,7 @@ private:
 
 		TextureSourceNew.Init(TargetImage.SizeX, TargetImage.SizeY, 1, 4);
 		int32 PixelCount = TextureSourceNew.GetPixelCount();
+		// @@!! if TextureSource is FImage this goes away
 		for (int PixelIndex = 0; PixelIndex < PixelCount ; PixelIndex++)
 		{
 			TextureSourceNew.SetPixel(PixelIndex, TargetImage.AsRGBA32F()[PixelIndex]);
diff --git a/Engine/Plugins/Enterprise/MDLImporter/Source/MDLImporter/Private/common/TextureSource.h b/Engine/Plugins/Enterprise/MDLImporter/Source/MDLImporter/Private/common/TextureSource.h
--- a/Engine/Plugins/Enterprise/MDLImporter/Source/MDLImporter/Private/common/TextureSource.h
+++ b/Engine/Plugins/Enterprise/MDLImporter/Source/MDLImporter/Private/common/TextureSource.h
@@ -6,6 +6,11 @@
 #include "Math/NumericLimits.h"
 #include "common/Utility.h"
 
+/**
+* TextureSource.h in the MDL importer ; not in the Engine
+* 
+*/
+
 namespace Common
 {
 	FORCEINLINE bool IsValueInRange(float Value)
@@ -16,6 +21,10 @@ namespace Common
 	// Converts a float texture to a corresponding texture source.
 	inline FTextureSource* CreateTextureSource(const float* InData, int InWidth, int InHeight, int InChannels, bool bFlipY)
 	{
+		// @@!! TODO use FImage / CopyImage
+		//	also why are we converting to G8 or RGBA16 here ?
+		// just use F32 ETextureSourceFormat ?
+
 		// use 16 bpp for linear textures support(i.e. if more than 1 channel)
 		const ETextureSourceFormat Format = InChannels == 1 ? TSF_G8 : TSF_RGBA16;
 		const int                  Size   = InWidth * InHeight * InChannels;
@@ -76,6 +85,8 @@ namespace Common
 						int DstOffset = Y * InWidth * 4 + X * 4;
 						int SrcOffset = SrcY * InWidth * 4 + X * 4;
 
+						// @@!! use FColor::QuantizeUNormFloatTo16
+						//	or just use FImage
 						check(IsValueInRange(InData[SrcOffset]));
 						check(IsValueInRange(InData[SrcOffset + 1]));
 						check(IsValueInRange(InData[SrcOffset + 2]));
diff --git a/Engine/Plugins/MovieScene/MovieRenderPipeline/Source/MovieRenderPipelineCore/Private/MoviePipelineImageQuantization.cpp b/Engine/Plugins/MovieScene/MovieRenderPipeline/Source/MovieRenderPipelineCore/Private/MoviePipelineImageQuantization.cpp
--- a/Engine/Plugins/MovieScene/MovieRenderPipeline/Source/MovieRenderPipelineCore/Private/MoviePipelineImageQuantization.cpp
+++ b/Engine/Plugins/MovieScene/MovieRenderPipeline/Source/MovieRenderPipelineCore/Private/MoviePipelineImageQuantization.cpp
@@ -14,6 +14,10 @@ namespace MoviePipeline
 DECLARE_CYCLE_STAT(TEXT("STAT_MoviePipeline_ImageQuantization"), STAT_ImageQuantization, STATGROUP_MoviePipeline);
 
 
+// @@!! TODO: get rid of this whole file
+//		use FImage, just call CopyImage
+
+
 template<class TToColorBitDepthType, typename TFromColorBitDepthType, typename TColorChannelNumericType>
 TArray<TToColorBitDepthType> ConvertLinearToLinearBitDepth(TFromColorBitDepthType* InColor, const int32 InCount)
 {
diff --git a/Engine/Plugins/MovieScene/MovieRenderPipeline/Source/MovieRenderPipelineRenderPasses/Private/MoviePipelineEXROutput.cpp b/Engine/Plugins/MovieScene/MovieRenderPipeline/Source/MovieRenderPipelineRenderPasses/Private/MoviePipelineEXROutput.cpp
--- a/Engine/Plugins/MovieScene/MovieRenderPipeline/Source/MovieRenderPipelineRenderPasses/Private/MoviePipelineEXROutput.cpp
+++ b/Engine/Plugins/MovieScene/MovieRenderPipeline/Source/MovieRenderPipelineRenderPasses/Private/MoviePipelineEXROutput.cpp
@@ -24,6 +24,8 @@ THIRD_PARTY_INCLUDES_END
 
 #if WITH_UNREALEXR
 
+// NOTE: see also ExrImageWrapper
+
 class FExrMemStreamOut : public Imf::OStream
 {
 public:
