nop comments only

diff --git a/Engine/Plugins/Enterprise/AxFImporter/Source/AxFImporter/Private/AxFImporter.cpp b/Engine/Plugins/Enterprise/AxFImporter/Source/AxFImporter/Private/AxFImporter.cpp
--- a/Engine/Plugins/Enterprise/AxFImporter/Source/AxFImporter/Private/AxFImporter.cpp
+++ b/Engine/Plugins/Enterprise/AxFImporter/Source/AxFImporter/Private/AxFImporter.cpp
@@ -99,8 +99,7 @@ public:
 
 	struct FAxFTextureSource
 	{
-		// @@!! just use float here, convert later
-		//  ideally replace this whole class with FImage
+		// @todo Oodle : just use float here, convert later ; replace this whole class with FImage
 		TArray<FFloat16> Pixels;
 
 		int32 Width;
@@ -2250,9 +2249,7 @@ private:
 			PixelValueMin = PixelValueMin.ComponentMin(PixelValue);
 		}
 
-		// @@!!
-		// float image to RGBA16
-		// todo: use FImage / CopyImage
+		// @todo Oodle : float image to RGBA16 ; just use FImage / CopyImage
 		TArray<uint16> PixelsCompressed;
 		if (!PixelValueMax.IsNearlyZero())
 		{
@@ -2271,10 +2268,10 @@ private:
 
 				uint16 One = 65535;
 
-				// @@!! precompute 1/Scale then clamp in [0,1]
+				// @todo Oodle : precompute 1/Scale then clamp in [0,1]
 				FVector C = (ClampVector(Color, FVector::ZeroVector, Scale) / Scale) * One; // clamp, to avoid integer overflow(which results in nasty pixels)
 
-				// @@!! incorrect conversion; use QuantizeUNormFloatTo16
+				// @todo Oodle : incorrect conversion; use QuantizeUNormFloatTo16
 				PixelsCompressed[PixelIndex * 4 + 0] = FMath::FloorToInt(C.X);
 				PixelsCompressed[PixelIndex * 4 + 1] = FMath::FloorToInt(C.Y);
 				PixelsCompressed[PixelIndex * 4 + 2] = FMath::FloorToInt(C.Z);
@@ -2372,7 +2369,7 @@ private:
 
 	void SetTextureSource(UTexture2D* Texture, FAxFTextureSource TextureSource)
 	{
-		//@@!! can I just make 32F ? no conversion?
+		//@todo Oodle : just make TSF_RGBA32F , no conversion
 		Texture->Source.Init(
 			TextureSource.Width,
 			TextureSource.Height,
@@ -2387,7 +2384,7 @@ private:
 	{
 		FImage Image;
 		Image.Init(TextureSource.Width, TextureSource.Height, ERawImageFormat::RGBA32F);
-		// @@!! if TextureSource is FImage this goes away
+		// @todo Oodle : if TextureSource is FImage this goes away
 		for (int PixelIndex = 0; PixelIndex < TextureSource.Width * TextureSource.Height; PixelIndex++)
 		{
 			Image.AsRGBA32F()[PixelIndex] = TextureSource.GetPixel(PixelIndex);
@@ -2412,7 +2409,7 @@ private:
 
 		TextureSourceNew.Init(TargetImage.SizeX, TargetImage.SizeY, 1, 4);
 		int32 PixelCount = TextureSourceNew.GetPixelCount();
-		// @@!! if TextureSource is FImage this goes away
+		// @todo Oodle : if TextureSource is FImage this goes away
 		for (int PixelIndex = 0; PixelIndex < PixelCount ; PixelIndex++)
 		{
 			TextureSourceNew.SetPixel(PixelIndex, TargetImage.AsRGBA32F()[PixelIndex]);
diff --git a/Engine/Source/Editor/UnrealEd/Private/EditorViewportClient.cpp b/Engine/Source/Editor/UnrealEd/Private/EditorViewportClient.cpp
--- a/Engine/Source/Editor/UnrealEd/Private/EditorViewportClient.cpp
+++ b/Engine/Source/Editor/UnrealEd/Private/EditorViewportClient.cpp
@@ -5839,7 +5839,7 @@ bool FEditorViewportClient::ProcessScreenShots(FViewport* InViewport)
 			TFuture<bool> CompletionFuture = HighResScreenshotConfig.ImageWriteQueue->Enqueue(MoveTemp(ImageTask));
 			if (CompletionFuture.IsValid())
 			{
-				// @@!! this queues it then immediately waits? what's the point of ImageWriteQueue then?
+				// this queues it then immediately waits? what's the point of ImageWriteQueue then?
 				// just use FImageUtils::Save
 				bIsScreenshotSaved = CompletionFuture.Get();
 			}
