Disable AxF importer for Windows Arm64

#rnx

diff --git a/Engine/Plugins/Enterprise/AxFImporter/AxFImporter.uplugin b/Engine/Plugins/Enterprise/AxFImporter/AxFImporter.uplugin
--- a/Engine/Plugins/Enterprise/AxFImporter/AxFImporter.uplugin
+++ b/Engine/Plugins/Enterprise/AxFImporter/AxFImporter.uplugin
@@ -24,6 +24,9 @@
 			"LoadingPhase": "Default",
 			"PlatformAllowList": [
 				"Win64"
+			],
+			"PlatformArchitectureDenyList": [
+				"Win64:arm64"
 			]
 		}
 	]
